Avoid dirtying the level when isolating an actor

#tests
#rnx

diff --git a/Engine/Plugins/Experimental/UserToolBoxBasicCommand/Source/UserToolBoxBasicCommand/Private/IsolateSelection.cpp b/Engine/Plugins/Experimental/UserToolBoxBasicCommand/Source/UserToolBoxBasicCommand/Private/IsolateSelection.cpp
--- a/Engine/Plugins/Experimental/UserToolBoxBasicCommand/Source/UserToolBoxBasicCommand/Private/IsolateSelection.cpp
+++ b/Engine/Plugins/Experimental/UserToolBoxBasicCommand/Source/UserToolBoxBasicCommand/Private/IsolateSelection.cpp
@@ -82,7 +82,7 @@ void UIsolateSelection::Execute()
 				continue;	
 			}
 
-			Actor->Modify();
+			SaveToTransactionBuffer(Actor, false);
 			Actor->SetIsTemporarilyHiddenInEditor(false);
 		}
 	}
@@ -102,7 +102,7 @@ void UIsolateSelection::Execute()
 					// Save the actor to the transaction buffer to support undo/redo, but do
 					// not call Modify, as we do not want to dirty the actor's package and
 					// we're only editing temporary, transient values
-					Actor->Modify();
+					SaveToTransactionBuffer(Actor, false);
 					if (Components.Num() > 0)
 					{
 						Actor->SetIsTemporarilyHiddenInEditor(true);
