Add an actor filter and new usertoolbox command to support automobile import issue

diff --git a/Engine/Plugins/Experimental/UserToolBoxBasicCommand/Source/UserToolBoxBasicCommand/Private/ActorFilter/BaseActorFilter.cpp b/Engine/Plugins/Experimental/UserToolBoxBasicCommand/Source/UserToolBoxBasicCommand/Private/ActorFilter/BaseActorFilter.cpp
--- a/Engine/Plugins/Experimental/UserToolBoxBasicCommand/Source/UserToolBoxBasicCommand/Private/ActorFilter/BaseActorFilter.cpp
+++ b/Engine/Plugins/Experimental/UserToolBoxBasicCommand/Source/UserToolBoxBasicCommand/Private/ActorFilter/BaseActorFilter.cpp
@@ -3,11 +3,14 @@
 
 #include "ActorFilter/BaseActorFilter.h"
 #include "DatasmithAssetUserData.h"
+#include "Components/StaticMeshComponent.h"
 #include "DatasmithContentBlueprintLibrary.h"
 #include "Widgets/Layout/SScrollBox.h"
 #include "Widgets/SWindow.h"
 #include "Widgets/Input/SButton.h"
 #include "Editor.h"
+#include "Engine/StaticMeshActor.h"
+
 void GetAttachedActors( AActor* Actor, bool bRecurseChildren, TSet< AActor* >& OutActors )
 {
 	TArray< AActor* > ChildrenActors;
@@ -267,3 +270,15 @@ TArray<AActor*> UHasMetadataByKeyAndValueDropDown::FilterImpl(const TArray<AActo
 	return ActorSortByValue.FindOrAdd(SelectedValue);
 
 }
+TArray<AActor*> UEmptyStaticMeshActor::FilterImpl(const TArray<AActor*>& Source)
+{
+	TArray<AActor*> Dest;
+	for (AActor* Cur:Source)
+	{
+		if (Cur->IsA(AStaticMeshActor::StaticClass()) && Cast<AStaticMeshActor>(Cur)->GetStaticMeshComponent()->GetStaticMesh()==nullptr)
+		{
+			Dest.Add(Cur);
+		}
+	}
+	return Dest;
+}
\ No newline at end of file
diff --git a/Engine/Plugins/Experimental/UserToolBoxBasicCommand/Source/UserToolBoxBasicCommand/Private/FillStaticMeshActor.cpp b/Engine/Plugins/Experimental/UserToolBoxBasicCommand/Source/UserToolBoxBasicCommand/Private/FillStaticMeshActor.cpp
new file mode 100644
--- /dev/null
+++ b/Engine/Plugins/Experimental/UserToolBoxBasicCommand/Source/UserToolBoxBasicCommand/Private/FillStaticMeshActor.cpp
@@ -0,0 +1,78 @@
+﻿// Copyright Epic Games, Inc. All Rights Reserved.
+
+
+#include "FillStaticMeshActor.h"
+#include "Engine/Selection.h"
+
+#include "Editor.h"
+#include "UserToolBoxBasicCommand.h"
+#include "Components/StaticMeshComponent.h"
+#include "AssetRegistry/AssetRegistryModule.h"
+#include "Engine/StaticMeshActor.h"
+
+
+void UFillStaticMeshActor::Execute()
+{
+	UWorld* World = GEngine->GetWorldContexts()[0].World();
+
+	USelection* SelectedActors = GEditor->GetSelectedActors();
+	TArray<AActor*> Actors;
+	SelectedActors->GetSelectedObjects(Actors);
+	FAssetRegistryModule& AssetRegistryModule = FModuleManager::LoadModuleChecked<FAssetRegistryModule>("AssetRegistry");
+	IAssetRegistry& AssetRegistry= AssetRegistryModule.Get();
+	FARFilter Filter;
+	Filter.bRecursivePaths=true;
+	if (!RootPaths.IsEmpty())
+	{
+		for (FString RootPath:RootPaths)
+		{
+			Filter.PackagePaths.Add(FName(RootPath));	
+		}
+			
+	}
+	else
+	{
+		Filter.PackagePaths.Add(FName("/Game"));	
+	}
+	
+	TArray<FAssetData> AssetDatas;
+	FTopLevelAssetPath ClassPath(FName("/Script/Engine"),FName("StaticMesh"));
+	Filter.ClassPaths.Add(ClassPath);
+	AssetRegistry.GetAssets(Filter,AssetDatas);
+	TMap<FName,FAssetData> MapOfStaticMeshByAssetName;
+	AssetRegistry.GetAssets(Filter,AssetDatas);
+	for (FAssetData& AssetData:AssetDatas)
+	{
+		MapOfStaticMeshByAssetName.Add(AssetData.AssetName,AssetData);
+	}
+	for (AActor* Actor: Actors)
+	{
+		if (Actor->IsA(AStaticMeshActor::StaticClass()))
+		{
+			AStaticMeshActor* StaticMeshActor=Cast<AStaticMeshActor>(Actor);
+			UStaticMesh* CurrentSM=StaticMeshActor->GetStaticMeshComponent()->GetStaticMesh();
+			if (CurrentSM==nullptr || !bAffectOnlyEmptyStaticMeshActor)
+			{
+				FString ActorLabel=Actor->GetActorLabel();
+				FAssetData* FoundedAssetData=MapOfStaticMeshByAssetName.Find(FName(ActorLabel));
+				if (FoundedAssetData!=nullptr)
+				{
+					UStaticMesh* StaticMesh=Cast<UStaticMesh>((*FoundedAssetData).GetAsset());
+					if (StaticMesh!=nullptr)
+					{
+						StaticMeshActor->GetStaticMeshComponent()->SetStaticMesh(StaticMesh);
+					}
+				}
+			}
+		}
+		
+	}
+}
+
+
+UFillStaticMeshActor::UFillStaticMeshActor()
+{
+	Name="Fill StaticMeshActor";
+	Tooltip="Fill SMA with empty Static Mesh with Mesh in the same name";
+	Category="Actor";
+}
diff --git a/Engine/Plugins/Experimental/UserToolBoxBasicCommand/Source/UserToolBoxBasicCommand/Public/ActorFilter/BaseActorFilter.h b/Engine/Plugins/Experimental/UserToolBoxBasicCommand/Source/UserToolBoxBasicCommand/Public/ActorFilter/BaseActorFilter.h
--- a/Engine/Plugins/Experimental/UserToolBoxBasicCommand/Source/UserToolBoxBasicCommand/Public/ActorFilter/BaseActorFilter.h
+++ b/Engine/Plugins/Experimental/UserToolBoxBasicCommand/Source/UserToolBoxBasicCommand/Public/ActorFilter/BaseActorFilter.h
@@ -137,3 +137,11 @@ public:
 	
 	virtual TArray<AActor*> FilterImpl(const TArray<AActor*>& Source) override;
 };
+UCLASS(Blueprintable)
+class UEmptyStaticMeshActor : public UBaseActorFilter
+{
+	GENERATED_BODY()
+public:
+	
+	virtual TArray<AActor*> FilterImpl(const TArray<AActor*>& Source) override;
+};
diff --git a/Engine/Plugins/Experimental/UserToolBoxBasicCommand/Source/UserToolBoxBasicCommand/Public/FillStaticMeshActor.h b/Engine/Plugins/Experimental/UserToolBoxBasicCommand/Source/UserToolBoxBasicCommand/Public/FillStaticMeshActor.h
new file mode 100644
--- /dev/null
+++ b/Engine/Plugins/Experimental/UserToolBoxBasicCommand/Source/UserToolBoxBasicCommand/Public/FillStaticMeshActor.h
@@ -0,0 +1,26 @@
+﻿// Copyright Epic Games, Inc. All Rights Reserved.
+
+#pragma once
+
+#include "CoreMinimal.h"
+#include "UTBBaseCommand.h"
+#include "FillStaticMeshActor.generated.h"
+
+/**
+ * 
+ */
+UCLASS()
+class USERTOOLBOXBASICCOMMAND_API UFillStaticMeshActor : public UUTBBaseCommand
+{
+	GENERATED_BODY()
+public:
+
+	UPROPERTY(EditAnywhere,BlueprintReadWrite,Category="FillStaticMeshActor")
+	TArray<FString>	RootPaths;
+	
+	UPROPERTY(EditAnywhere,BlueprintReadWrite,Category="FillStaticMeshActor")
+	bool	bAffectOnlyEmptyStaticMeshActor=true;
+	
+	void Execute() override;
+	UFillStaticMeshActor();
+};
